Fix copyright notice.

diff --git a/Engine/Plugins/Compositing/OpenCVLensDistortion/Source/OpenCVLensCalibration/Private/OpenCVLensCalibrator.cpp b/Engine/Plugins/Compositing/OpenCVLensDistortion/Source/OpenCVLensCalibration/Private/OpenCVLensCalibrator.cpp
--- a/Engine/Plugins/Compositing/OpenCVLensDistortion/Source/OpenCVLensCalibration/Private/OpenCVLensCalibrator.cpp
+++ b/Engine/Plugins/Compositing/OpenCVLensDistortion/Source/OpenCVLensCalibration/Private/OpenCVLensCalibrator.cpp
@@ -1,4 +1,4 @@
-// Copyright 1998-2017 Epic Games, Inc. All Rights Reserved.
+// Copyright 1998-2018 Epic Games, Inc. All Rights Reserved.
 
 #include "OpenCVLensCalibrator.h"
 
