Updating the MRCalibration project to compile/run after Enterprise moved the OpenCV calibration out to its own plugin.

diff --git a/Engine/Plugins/Compositing/OpenCVLensDistortion/Config/BaseOpenCVLensDistortion.ini b/Engine/Plugins/Compositing/OpenCVLensDistortion/Config/BaseOpenCVLensDistortion.ini
new file mode 100644
--- /dev/null
+++ b/Engine/Plugins/Compositing/OpenCVLensDistortion/Config/BaseOpenCVLensDistortion.ini
@@ -0,0 +1,12 @@
+[CoreRedirects]
+
+;
+; 4.20 Redirects
+;
+
++PackageRedirects=(OldName="/Script/MixedRealityCaptureCalibration", NewName="/Script/OpenCVLensCalibration")
+
++StructRedirects=(OldName="MrcLensDistortion", NewName="/Script/OpenCVLensDistortion.OpenCVLensDistortionParameters")
++ClassRedirects=(OldName="MrcOpenCVCalibrator", NewName="/Script/OpenCVLensCalibration.OpenCVLensCalibrator")
++PropertyRedirects=(OldName="FeedRenderTarget.TexRT", NewName="FeedRenderTarget.TextureRenderTarget")
++PropertyRedirects=(OldName="CalculateLensParameters.LensDistortion", NewName="CalculateLensParameters.LensDistortionParameters")
\ No newline at end of file
