Copied 4393087 Patrick.Boutot

Add a shader source directory for the OpenCVLensDistortion plugin.

diff --git a/Engine/Plugins/Compositing/OpenCVLensDistortion/Source/OpenCVLensDistortion/OpenCVLensDistortion.Build.cs b/Engine/Plugins/Compositing/OpenCVLensDistortion/Source/OpenCVLensDistortion/OpenCVLensDistortion.Build.cs
--- a/Engine/Plugins/Compositing/OpenCVLensDistortion/Source/OpenCVLensDistortion/OpenCVLensDistortion.Build.cs
+++ b/Engine/Plugins/Compositing/OpenCVLensDistortion/Source/OpenCVLensDistortion/OpenCVLensDistortion.Build.cs
@@ -32,6 +32,7 @@ namespace UnrealBuildTool.Rules
 					"Core",
 					"CoreUObject",
 					"Engine",
+					"Projects",
 					"RenderCore",
 					"RHI",
 					"ShaderCore",
diff --git a/Engine/Plugins/Compositing/OpenCVLensDistortion/Source/OpenCVLensDistortion/Private/OpenCVLensDistortionModule.cpp b/Engine/Plugins/Compositing/OpenCVLensDistortion/Source/OpenCVLensDistortion/Private/OpenCVLensDistortionModule.cpp
--- a/Engine/Plugins/Compositing/OpenCVLensDistortion/Source/OpenCVLensDistortion/Private/OpenCVLensDistortionModule.cpp
+++ b/Engine/Plugins/Compositing/OpenCVLensDistortion/Source/OpenCVLensDistortion/Private/OpenCVLensDistortionModule.cpp
@@ -2,7 +2,10 @@
 
 #include "IOpenCVLensDistortionModule.h"
 
+#include "Interfaces/IPluginManager.h"
 #include "Logging/LogMacros.h"
+#include "Misc/Paths.h"
+#include "ShaderCore.h"
 
 
 DEFINE_LOG_CATEGORY(LogOpenCVLensDistortion)
@@ -11,7 +14,13 @@ DEFINE_LOG_CATEGORY(LogOpenCVLensDistortion)
 // FOpenCVLensDistortionModule
 class FOpenCVLensDistortionModule : public IOpenCVLensDistortionModule
 {
-
+public:
+	virtual void StartupModule() override
+	{
+		// Maps virtual shader source directory /Plugin/OpenCVLensDistortion to the plugin's actual Shaders directory.
+		FString PluginShaderDir = FPaths::Combine(IPluginManager::Get().FindPlugin(TEXT("OpenCVLensDistortion"))->GetBaseDir(), TEXT("Shaders"));
+		AddShaderSourceDirectoryMapping(TEXT("/Plugin/OpenCVLensDistortion"), PluginShaderDir);
+	}
 };
 
 //////////////////////////////////////////////////////////////////////////
