Changes required by 4073067 - "Replaces automatic plugin and project shader directory mapping with manual aproach in StartupModule()"

diff --git a/Engine/Plugins/Compositing/OpenCVLensDistortion/Source/OpenCVLensDistortion/OpenCVLensDistortion.Build.cs b/Engine/Plugins/Compositing/OpenCVLensDistortion/Source/OpenCVLensDistortion/OpenCVLensDistortion.Build.cs
--- a/Engine/Plugins/Compositing/OpenCVLensDistortion/Source/OpenCVLensDistortion/OpenCVLensDistortion.Build.cs
+++ b/Engine/Plugins/Compositing/OpenCVLensDistortion/Source/OpenCVLensDistortion/OpenCVLensDistortion.Build.cs
@@ -35,6 +35,7 @@ namespace UnrealBuildTool.Rules
 					"RenderCore",
 					"RHI",
 					"ShaderCore",
+					"Projects",
 				}
 			);
 		}
diff --git a/Engine/Plugins/Compositing/OpenCVLensDistortion/Source/OpenCVLensDistortion/Private/OpenCVLensDistortionModule.cpp b/Engine/Plugins/Compositing/OpenCVLensDistortion/Source/OpenCVLensDistortion/Private/OpenCVLensDistortionModule.cpp
--- a/Engine/Plugins/Compositing/OpenCVLensDistortion/Source/OpenCVLensDistortion/Private/OpenCVLensDistortionModule.cpp
+++ b/Engine/Plugins/Compositing/OpenCVLensDistortion/Source/OpenCVLensDistortion/Private/OpenCVLensDistortionModule.cpp
@@ -3,6 +3,9 @@
 #include "IOpenCVLensDistortionModule.h"
 
 #include "Logging/LogMacros.h"
+#include "Misc/Paths.h"
+#include "Interfaces/IPluginManager.h"
+#include "ShaderCore.h"
 
 
 DEFINE_LOG_CATEGORY(LogOpenCVLensDistortion)
@@ -19,3 +22,8 @@ class FOpenCVLensDistortionModule : public IOpenCVLensDistortionModule
 IMPLEMENT_MODULE(FOpenCVLensDistortionModule, OpenCVLensDistortion);
 
 
+void IOpenCVLensDistortionModule::StartupModule()
+{
+	FString PluginShaderDir = FPaths::Combine(IPluginManager::Get().FindPlugin(TEXT("OpenCVLensDistortion"))->GetBaseDir(), TEXT("Shaders"));
+	AddShaderSourceDirectoryMapping(TEXT("/Plugin/OpenCVLensDistortion"), PluginShaderDir);
+}
diff --git a/Engine/Plugins/Compositing/OpenCVLensDistortion/Source/OpenCVLensDistortion/Public/IOpenCVLensDistortionModule.h b/Engine/Plugins/Compositing/OpenCVLensDistortion/Source/OpenCVLensDistortion/Public/IOpenCVLensDistortionModule.h
--- a/Engine/Plugins/Compositing/OpenCVLensDistortion/Source/OpenCVLensDistortion/Public/IOpenCVLensDistortionModule.h
+++ b/Engine/Plugins/Compositing/OpenCVLensDistortion/Source/OpenCVLensDistortion/Public/IOpenCVLensDistortionModule.h
@@ -19,6 +19,9 @@ class IOpenCVLensDistortionModule : public IModuleInterface
 
 public:
 
+	// IModuleInterface
+	virtual void StartupModule() override;
+
 	/**
 	 * Singleton-like access to this module's interface.  This is just for convenience!
 	 * Beware of calling this during the shutdown phase, though.  Your module might have been unloaded already.
