Use reinterpret_cast as the texture bulkdata stores the data as uint8 but the code wants uint16.

#rnx

diff --git a/Engine/Plugins/Compositing/OpenCVLensDistortion/Source/OpenCVLensDistortion/Private/OpenCVLensDistortionParameters.cpp b/Engine/Plugins/Compositing/OpenCVLensDistortion/Source/OpenCVLensDistortion/Private/OpenCVLensDistortionParameters.cpp
--- a/Engine/Plugins/Compositing/OpenCVLensDistortion/Source/OpenCVLensDistortion/Private/OpenCVLensDistortionParameters.cpp
+++ b/Engine/Plugins/Compositing/OpenCVLensDistortion/Source/OpenCVLensDistortion/Private/OpenCVLensDistortionParameters.cpp
@@ -100,7 +100,7 @@ UTexture2D* FOpenCVLensDistortionParameters::CreateUndistortUVDisplacementMap(co
 
 	// Lock the texture so it can be modified
 	FTexture2DMipMap& Mip = Result->PlatformData->Mips[0];
-	uint16* MipData = static_cast<uint16*>(Mip.BulkData.Lock(LOCK_READ_WRITE));
+	uint16* MipData = reinterpret_cast<uint16*>(Mip.BulkData.Lock(LOCK_READ_WRITE));
 	check(MipData);
 
 	// Go through each pixel and change to normalized displacement value
