Updated Comment

diff --git a/Engine/Plugins/Compositing/OpenCVLensDistortion/Source/OpenCVLensCalibration/Private/OpenCVLensCalibrator.cpp b/Engine/Plugins/Compositing/OpenCVLensDistortion/Source/OpenCVLensCalibration/Private/OpenCVLensCalibrator.cpp
--- a/Engine/Plugins/Compositing/OpenCVLensDistortion/Source/OpenCVLensCalibration/Private/OpenCVLensCalibrator.cpp
+++ b/Engine/Plugins/Compositing/OpenCVLensDistortion/Source/OpenCVLensCalibration/Private/OpenCVLensCalibrator.cpp
@@ -230,7 +230,7 @@ bool UOpenCVLensCalibrator::CalculateLensParameters(FOpenCVLensDistortionParamet
 		return false;
 	}
 
-	// Convert the params to a UE4 struct
+	// Convert the params to a Unreal Editor struct
 	{
 		//Fisheye camera model vs pinhole camera model differs slightly in parameter assignment
 		//Pinhole can have up to 6 radial distortion parameters and 2 tangential parameters
