Moved common parameters from FOpenCVLensDistortionParameters out of the Lens Distortion plugin to the OpenCV plugin so that any other plugin can easily use it

diff --git a/Engine/Plugins/Compositing/OpenCVLensDistortion/Source/OpenCVLensDistortion/Private/OpenCVLensDistortionParameters.cpp b/Engine/Plugins/Compositing/OpenCVLensDistortion/Source/OpenCVLensDistortion/Private/OpenCVLensDistortionParameters.cpp
--- a/Engine/Plugins/Compositing/OpenCVLensDistortion/Source/OpenCVLensDistortion/Private/OpenCVLensDistortionParameters.cpp
+++ b/Engine/Plugins/Compositing/OpenCVLensDistortion/Source/OpenCVLensDistortion/Private/OpenCVLensDistortionParameters.cpp
@@ -16,46 +16,6 @@ OPENCV_INCLUDES_END
 #endif
 
 
-#if WITH_OPENCV
-cv::Mat FOpenCVLensDistortionParameters::ConvertToOpenCVDistortionCoefficients() const
-{
-	if (bUseFisheyeModel)
-	{
-		cv::Mat DistortionCoefficients(1, 4, CV_64F);
-		DistortionCoefficients.at<double>(0) = K1;
-		DistortionCoefficients.at<double>(1) = K2;
-		DistortionCoefficients.at<double>(2) = K3;
-		DistortionCoefficients.at<double>(3) = K4;
-		return DistortionCoefficients;
-	}
-	else
-	{
-		cv::Mat DistortionCoefficients(1, 8, CV_64F);
-		DistortionCoefficients.at<double>(0) = K1;
-		DistortionCoefficients.at<double>(1) = K2;
-		DistortionCoefficients.at<double>(2) = P1;
-		DistortionCoefficients.at<double>(3) = P2;
-		DistortionCoefficients.at<double>(4) = K3;
-		DistortionCoefficients.at<double>(5) = K4;
-		DistortionCoefficients.at<double>(6) = K5;
-		DistortionCoefficients.at<double>(7) = K6;
-		return DistortionCoefficients;
-	}
-}
-#endif
-
-#if WITH_OPENCV
-cv::Mat FOpenCVLensDistortionParameters::CreateOpenCVCameraMatrix(const FVector2D& InImageSize) const
-{
-	cv::Mat CameraMatrix = cv::Mat::eye(3, 3, CV_64F);
-	CameraMatrix.at<double>(0, 0) = F.X * InImageSize.X;
-	CameraMatrix.at<double>(1, 1) = F.Y * InImageSize.Y;
-	CameraMatrix.at<double>(0, 2) = C.X * InImageSize.X;
-	CameraMatrix.at<double>(1, 2) = C.Y * InImageSize.Y;
-	return CameraMatrix;
-}
-#endif
-
 UTexture2D* FOpenCVLensDistortionParameters::CreateUndistortUVDisplacementMap(const FIntPoint& InImageSize, const float InCroppingFactor, FOpenCVCameraViewInfo& OutCameraViewInfo) const
 {
 #if WITH_OPENCV
diff --git a/Engine/Plugins/Compositing/OpenCVLensDistortion/Source/OpenCVLensDistortion/Public/OpenCVLensDistortionParameters.h b/Engine/Plugins/Compositing/OpenCVLensDistortion/Source/OpenCVLensDistortion/Public/OpenCVLensDistortionParameters.h
--- a/Engine/Plugins/Compositing/OpenCVLensDistortion/Source/OpenCVLensDistortion/Public/OpenCVLensDistortionParameters.h
+++ b/Engine/Plugins/Compositing/OpenCVLensDistortion/Source/OpenCVLensDistortion/Public/OpenCVLensDistortionParameters.h
@@ -4,8 +4,11 @@
 
 #include "CoreMinimal.h"
 
-#if WITH_OPENCV
+// Keep OpenCVHelper.h outside of the WITH_OPENCV define so that FOpenCVLensDistortionParametersBase is included
+// OpenCVHelper.h has its own WITH_OPENCV guards so having it guarded again here isn't necessary
 #include "OpenCVHelper.h"
+
+#if WITH_OPENCV
 OPENCV_INCLUDES_START
 #undef check // the check macro causes problems with opencv headers
 #include "opencv2/core/core.hpp"
@@ -52,26 +55,11 @@ struct OPENCVLENSDISTORTION_API FOpenCVCameraViewInfo
  * where F and C are normalized.
  */
 USTRUCT(BlueprintType)
-struct OPENCVLENSDISTORTION_API FOpenCVLensDistortionParameters
+struct OPENCVLENSDISTORTION_API FOpenCVLensDistortionParameters : public FOpenCVLensDistortionParametersBase
 {
 	GENERATED_USTRUCT_BODY()
 
 public:
-	FOpenCVLensDistortionParameters()
-		: K1(0.f)
-		, K2(0.f)
-		, P1(0.f)
-		, P2(0.f)
-		, K3(0.f)
-		, K4(0.f)
-		, K5(0.f)
-		, K6(0.f)
-		, F(FVector2D(1.f, 1.f))
-		, C(FVector2D(0.5f, 0.5f))
-		, bUseFisheyeModel(false)
-	{
-	}
-
 	/** Draws UV displacement map within the output render target.
 	 * - Red & green channels hold the distort to undistort displacement;
 	 * - Blue & alpha channels hold the undistort to distort displacement.
@@ -89,104 +77,4 @@ public:
 	 * @param OutCameraViewInfo Information computed by OpenCV about the undistorted space. Can be used with SceneCapture to adjust FOV.
 	 */
 	UTexture2D* CreateUndistortUVDisplacementMap(const FIntPoint& InImageSize, const float InCroppingFactor, FOpenCVCameraViewInfo& OutCameraViewInfo) const;
-
-public:
-
-	/** Compare two lens distortion models and return whether they are equal. */
-	bool operator == (const FOpenCVLensDistortionParameters& Other) const
-	{
-		return (K1 == Other.K1 &&
-				K2 == Other.K2 &&
-				P1 == Other.P1 &&
-				P2 == Other.P2 &&
-				K3 == Other.K3 &&
-				K4 == Other.K4 &&
-				K5 == Other.K5 &&
-				K6 == Other.K6 &&
-				F == Other.F &&
-				C == Other.C &&
-				bUseFisheyeModel == Other.bUseFisheyeModel);
-	}
-
-	/** Compare two lens distortion models and return whether they are different. */
-	bool operator != (const FOpenCVLensDistortionParameters& Other) const
-	{
-		return !(*this == Other);
-	}
-
-	/** Returns true if lens distortion parameters are for identity lens (or default parameters) */
-	bool IsIdentity() const
-	{
-		return (K1 == 0.0f &&
-				K2 == 0.0f &&
-				P1 == 0.0f &&
-				P2 == 0.0f &&
-				K3 == 0.0f &&
-				K4 == 0.0f &&
-				K5 == 0.0f &&
-				K6 == 0.0f &&
-				F == FVector2D(1.0f, 1.0f) &&
-				C == FVector2D(0.5f, 0.5f));
-	}
-
-	bool IsSet() const 
-	{
-		return *this != FOpenCVLensDistortionParameters();
-	}
-
-private:
-
-#if WITH_OPENCV
-	/** Convert internal coefficients to OpenCV matrix representation */
-	cv::Mat ConvertToOpenCVDistortionCoefficients() const;
-
-	/** Convert internal normalized camera matrix to OpenCV pixel scaled matrix representation. */
-	cv::Mat CreateOpenCVCameraMatrix(const FVector2D& InImageSize) const;
-#endif //WITH_OPENCV
-
-public:
-
-	/** Radial parameter #1. */
-	UPROPERTY(Interp, EditAnywhere, BlueprintReadWrite, Category = "Lens Distortion|Parameters")
-	float K1;
-
-	/** Radial parameter #2. */
-	UPROPERTY(Interp, EditAnywhere, BlueprintReadWrite, Category = "Lens Distortion|Parameters")
-	float K2;
-
-	/** Tangential parameter #1. */
-	UPROPERTY(Interp, EditAnywhere, BlueprintReadWrite, Category = "Lens Distortion|Parameters")
-	float P1;
-
-	/** Tangential parameter #2. */
-	UPROPERTY(Interp, EditAnywhere, BlueprintReadWrite, Category = "Lens Distortion|Parameters")
-	float P2;
-
-	/** Radial parameter #3. */
-	UPROPERTY(Interp, EditAnywhere, BlueprintReadWrite, Category = "Lens Distortion|Parameters")
-	float K3;
-
-	/** Radial parameter #4. */
-	UPROPERTY(Interp, EditAnywhere, BlueprintReadWrite, Category = "Lens Distortion|Parameters")
-	float K4;
-
-	/** Radial parameter #5. */
-	UPROPERTY(Interp, EditAnywhere, BlueprintReadWrite, Category = "Lens Distortion|Parameters")
-	float K5;
-	
-	/** Radial parameter #6. */
-	UPROPERTY(Interp, EditAnywhere, BlueprintReadWrite, Category = "Lens Distortion|Parameters")
-	float K6;
-
-	/** Camera matrix's normalized Fx and Fy. */
-	UPROPERTY(Interp, EditAnywhere, BlueprintReadWrite, Category = "Lens Distortion|Parameters")
-	FVector2D F;
-
-	/** Camera matrix's normalized Cx and Cy. */
-	UPROPERTY(Interp, EditAnywhere, BlueprintReadWrite, Category = "Lens Distortion|Parameters")
-	FVector2D C;
-
-	/** Camera lens needs Fisheye camera model. */
-	UPROPERTY(EditAnywhere, BlueprintReadWrite, Category = "Lens Distortion|Parameters")
-	bool bUseFisheyeModel;
 };
diff --git a/Engine/Plugins/Media/MediaFrameworkUtilities/Source/MediaFrameworkUtilities/MediaFrameworkUtilities.Build.cs b/Engine/Plugins/Media/MediaFrameworkUtilities/Source/MediaFrameworkUtilities/MediaFrameworkUtilities.Build.cs
--- a/Engine/Plugins/Media/MediaFrameworkUtilities/Source/MediaFrameworkUtilities/MediaFrameworkUtilities.Build.cs
+++ b/Engine/Plugins/Media/MediaFrameworkUtilities/Source/MediaFrameworkUtilities/MediaFrameworkUtilities.Build.cs
@@ -10,6 +10,7 @@ namespace UnrealBuildTool.Rules
 				new string[]
 				{
 					"OpenCVLensDistortion",
+					"OpenCVHelper",
 					"TimeManagement",
 				}
 			);
diff --git a/Engine/Plugins/Runtime/MixedRealityCaptureFramework/Source/MixedRealityCaptureFramework/MixedRealityCaptureFramework.Build.cs b/Engine/Plugins/Runtime/MixedRealityCaptureFramework/Source/MixedRealityCaptureFramework/MixedRealityCaptureFramework.Build.cs
--- a/Engine/Plugins/Runtime/MixedRealityCaptureFramework/Source/MixedRealityCaptureFramework/MixedRealityCaptureFramework.Build.cs
+++ b/Engine/Plugins/Runtime/MixedRealityCaptureFramework/Source/MixedRealityCaptureFramework/MixedRealityCaptureFramework.Build.cs
@@ -34,6 +34,7 @@ public class MixedRealityCaptureFramework : ModuleRules
                 "MediaUtils",
 				"RenderCore",
 				"OpenCVLensDistortion",
+				"OpenCVHelper",
 			}
 		);
 
diff --git a/Engine/Plugins/Runtime/OpenCV/Source/OpenCVHelper/Private/OpenCVHelper.cpp b/Engine/Plugins/Runtime/OpenCV/Source/OpenCVHelper/Private/OpenCVHelper.cpp
--- a/Engine/Plugins/Runtime/OpenCV/Source/OpenCVHelper/Private/OpenCVHelper.cpp
+++ b/Engine/Plugins/Runtime/OpenCV/Source/OpenCVHelper/Private/OpenCVHelper.cpp
@@ -12,6 +12,41 @@ OPENCV_INCLUDES_START
 #include "opencv2/opencv.hpp"
 OPENCV_INCLUDES_END
 
+cv::Mat FOpenCVLensDistortionParametersBase::ConvertToOpenCVDistortionCoefficients() const
+{
+	if (bUseFisheyeModel)
+	{
+		cv::Mat DistortionCoefficients(1, 4, CV_64F);
+		DistortionCoefficients.at<double>(0) = K1;
+		DistortionCoefficients.at<double>(1) = K2;
+		DistortionCoefficients.at<double>(2) = K3;
+		DistortionCoefficients.at<double>(3) = K4;
+		return DistortionCoefficients;
+	}
+	else
+	{
+		cv::Mat DistortionCoefficients(1, 8, CV_64F);
+		DistortionCoefficients.at<double>(0) = K1;
+		DistortionCoefficients.at<double>(1) = K2;
+		DistortionCoefficients.at<double>(2) = P1;
+		DistortionCoefficients.at<double>(3) = P2;
+		DistortionCoefficients.at<double>(4) = K3;
+		DistortionCoefficients.at<double>(5) = K4;
+		DistortionCoefficients.at<double>(6) = K5;
+		DistortionCoefficients.at<double>(7) = K6;
+		return DistortionCoefficients;
+	}
+}
+
+cv::Mat FOpenCVLensDistortionParametersBase::CreateOpenCVCameraMatrix(const FVector2D& InImageSize) const
+{
+	cv::Mat CameraMatrix = cv::Mat::eye(3, 3, CV_64F);
+	CameraMatrix.at<double>(0, 0) = F.X * InImageSize.X;
+	CameraMatrix.at<double>(1, 1) = F.Y * InImageSize.Y;
+	CameraMatrix.at<double>(0, 2) = C.X * InImageSize.X;
+	CameraMatrix.at<double>(1, 2) = C.Y * InImageSize.Y;
+	return CameraMatrix;
+}
 
 UTexture2D* FOpenCVHelper::TextureFromCvMat(cv::Mat& Mat, const FString* PackagePath, const FName* TextureName)
 {
diff --git a/Engine/Plugins/Runtime/OpenCV/Source/OpenCVHelper/Public/OpenCVHelper.h b/Engine/Plugins/Runtime/OpenCV/Source/OpenCVHelper/Public/OpenCVHelper.h
--- a/Engine/Plugins/Runtime/OpenCV/Source/OpenCVHelper/Public/OpenCVHelper.h
+++ b/Engine/Plugins/Runtime/OpenCV/Source/OpenCVHelper/Public/OpenCVHelper.h
@@ -21,6 +21,8 @@ cannot undefine the value.
 
 */
 
+#include "OpenCVHelper.generated.h"
+
 #if PLATFORM_WINDOWS
 #define OPENCV_INCLUDES_START THIRD_PARTY_INCLUDES_START \
 	__pragma(warning(disable: 4190))  /* 'identifier1' has C-linkage specified, but returns UDT 'identifier2' which is incompatible with C */ \
@@ -74,3 +76,130 @@ public:
 };
 
 #endif //WITH_OPENCV
+
+/**
+ * Mathematic camera model for lens distortion/undistortion.
+ * Camera matrix =
+ *  | F.X  0  C.x |
+ *  |  0  F.Y C.Y |
+ *  |  0   0   1  |
+ * where F and C are normalized.
+ */
+USTRUCT(BlueprintType)
+struct OPENCVHELPER_API FOpenCVLensDistortionParametersBase
+{
+	GENERATED_USTRUCT_BODY()
+
+public:
+	FOpenCVLensDistortionParametersBase()
+		: K1(0.f)
+		, K2(0.f)
+		, P1(0.f)
+		, P2(0.f)
+		, K3(0.f)
+		, K4(0.f)
+		, K5(0.f)
+		, K6(0.f)
+		, F(FVector2D(1.f, 1.f))
+		, C(FVector2D(0.5f, 0.5f))
+		, bUseFisheyeModel(false)
+	{
+	}
+
+public:
+#if WITH_OPENCV
+	/** Convert internal coefficients to OpenCV matrix representation */
+	cv::Mat ConvertToOpenCVDistortionCoefficients() const;
+
+	/** Convert internal normalized camera matrix to OpenCV pixel scaled matrix representation. */
+	cv::Mat CreateOpenCVCameraMatrix(const FVector2D& InImageSize) const;
+#endif //WITH_OPENCV
+
+public:
+	/** Compare two lens distortion models and return whether they are equal. */
+	bool operator == (const FOpenCVLensDistortionParametersBase& Other) const
+	{
+		return (K1 == Other.K1 &&
+			K2 == Other.K2 &&
+			P1 == Other.P1 &&
+			P2 == Other.P2 &&
+			K3 == Other.K3 &&
+			K4 == Other.K4 &&
+			K5 == Other.K5 &&
+			K6 == Other.K6 &&
+			F == Other.F &&
+			C == Other.C &&
+			bUseFisheyeModel == Other.bUseFisheyeModel);
+	}
+
+	/** Compare two lens distortion models and return whether they are different. */
+	bool operator != (const FOpenCVLensDistortionParametersBase& Other) const
+	{
+		return !(*this == Other);
+	}
+
+	/** Returns true if lens distortion parameters are for identity lens (or default parameters) */
+	bool IsIdentity() const
+	{
+		return (K1 == 0.0f &&
+			K2 == 0.0f &&
+			P1 == 0.0f &&
+			P2 == 0.0f &&
+			K3 == 0.0f &&
+			K4 == 0.0f &&
+			K5 == 0.0f &&
+			K6 == 0.0f &&
+			F == FVector2D(1.0f, 1.0f) &&
+			C == FVector2D(0.5f, 0.5f));
+	}
+
+	bool IsSet() const
+	{
+		return *this != FOpenCVLensDistortionParametersBase();
+	}
+
+public:
+	/** Radial parameter #1. */
+	UPROPERTY(Interp, EditAnywhere, BlueprintReadWrite, Category = "Lens Distortion|Parameters")
+	float K1;
+
+	/** Radial parameter #2. */
+	UPROPERTY(Interp, EditAnywhere, BlueprintReadWrite, Category = "Lens Distortion|Parameters")
+	float K2;
+
+	/** Tangential parameter #1. */
+	UPROPERTY(Interp, EditAnywhere, BlueprintReadWrite, Category = "Lens Distortion|Parameters")
+	float P1;
+
+	/** Tangential parameter #2. */
+	UPROPERTY(Interp, EditAnywhere, BlueprintReadWrite, Category = "Lens Distortion|Parameters")
+	float P2;
+
+	/** Radial parameter #3. */
+	UPROPERTY(Interp, EditAnywhere, BlueprintReadWrite, Category = "Lens Distortion|Parameters")
+	float K3;
+
+	/** Radial parameter #4. */
+	UPROPERTY(Interp, EditAnywhere, BlueprintReadWrite, Category = "Lens Distortion|Parameters")
+	float K4;
+
+	/** Radial parameter #5. */
+	UPROPERTY(Interp, EditAnywhere, BlueprintReadWrite, Category = "Lens Distortion|Parameters")
+	float K5;
+	
+	/** Radial parameter #6. */
+	UPROPERTY(Interp, EditAnywhere, BlueprintReadWrite, Category = "Lens Distortion|Parameters")
+	float K6;
+
+	/** Camera matrix's normalized Fx and Fy. */
+	UPROPERTY(Interp, EditAnywhere, BlueprintReadWrite, Category = "Lens Distortion|Parameters")
+	FVector2D F;
+
+	/** Camera matrix's normalized Cx and Cy. */
+	UPROPERTY(Interp, EditAnywhere, BlueprintReadWrite, Category = "Lens Distortion|Parameters")
+	FVector2D C;
+
+	/** Camera lens needs Fisheye camera model. */
+	UPROPERTY(EditAnywhere, BlueprintReadWrite, Category = "Lens Distortion|Parameters")
+	bool bUseFisheyeModel;
+};
